- First checkin for the new DesktopTargetPlatform that will be able to cook packages for Windows, Mac, and Linux. So far, cooking on Windows and Mac and run on each other (with -opengl on Windows, since Mac can't currently compile D3D shaders)

diff --git a/Engine/Source/Developer/AudioFormatOpus/AudioFormatOpus.Build.cs b/Engine/Source/Developer/AudioFormatOpus/AudioFormatOpus.Build.cs
--- a/Engine/Source/Developer/AudioFormatOpus/AudioFormatOpus.Build.cs
+++ b/Engine/Source/Developer/AudioFormatOpus/AudioFormatOpus.Build.cs
@@ -17,9 +17,9 @@ public class AudioFormatOpus : ModuleRules
 
 		if ((Target.Platform == UnrealTargetPlatform.Win64) ||
 			(Target.Platform == UnrealTargetPlatform.Win32) ||
-            (Target.Platform == UnrealTargetPlatform.Linux)
-			//(Target.Platform == UnrealTargetPlatform.Mac) ||
-            //((Target.Platform == UnrealTargetPlatform.HTML5 && Target.Architecture == "-win32"))
+            (Target.Platform == UnrealTargetPlatform.Linux) ||
+			(Target.Platform == UnrealTargetPlatform.Mac)
+            //(Target.Platform == UnrealTargetPlatform.HTML5 && Target.Architecture == "-win32")
             )
 		{
 			AddThirdPartyPrivateStaticDependencies(Target, 
diff --git a/Engine/Source/Developer/DesktopPlatform/Private/PlatformInfo.cpp b/Engine/Source/Developer/DesktopPlatform/Private/PlatformInfo.cpp
--- a/Engine/Source/Developer/DesktopPlatform/Private/PlatformInfo.cpp
+++ b/Engine/Source/Developer/DesktopPlatform/Private/PlatformInfo.cpp
@@ -112,6 +112,8 @@ static const FPlatformInfo PlatformInfoArray[] = {
 
 	BuildPlatformInfo(TEXT("XboxOne"),					TEXT("XboxOne"),			LOCTEXT("XboxOne", "Xbox One"),								EPlatformType::Game,		EPlatformFlags::None,			FPlatformIconPaths(TEXT("Launcher/XboxOneTarget/Platform_XboxOne_24x"), TEXT("Launcher/XboxOneTarget/Platform_XboxOne_128x")),						TEXT(""),												TEXT("XboxOne"),	EPlatformSDKStatus::Unknown,	TEXT("Platforms/XboxOne/GettingStarted"),																IsAvailableOnWindows,						false,				false),
 
+	BuildPlatformInfo(TEXT("Desktop"),					TEXT("Desktop"),			LOCTEXT("DesktopTargetPlatDisplay", "Desktop (Win/Mac/Linux)"),	EPlatformType::Game,	EPlatformFlags::None,			FPlatformIconPaths(TEXT("Launcher/DesktopTarget/Platform_Desktop_24x"), TEXT("Launcher/DesktopTarget/Platform_Desktop_128x")),						TEXT(""),												TEXT(""),			EPlatformSDKStatus::Unknown,	TEXT(""),																							IsAvailableOnWindows || IsAvailableOnMac || IsAvailableOnLinux,	true, false),
+
 	//BuildPlatformInfo(TEXT("WinRT"),					TEXT("WinRT"),				LOCTEXT("WinRT", "Windows RT"),								EPlatformType::Game,		EPlatformFlags::None,			FPlatformIconPaths(TEXT("Launcher/WindowsTarget/Platform_WindowsNoEditor_24x"), TEXT("Launcher/WindowsTarget/Platform_WindowsNoEditor_128x")),		TEXT("")),
 };
 
diff --git a/Engine/Source/Developer/DesktopTargetPlatform/DesktopTargetPlatform.Build.cs b/Engine/Source/Developer/DesktopTargetPlatform/DesktopTargetPlatform.Build.cs
new file mode 100644
--- /dev/null
+++ b/Engine/Source/Developer/DesktopTargetPlatform/DesktopTargetPlatform.Build.cs
@@ -0,0 +1,33 @@
+// Copyright 1998-2015 Epic Games, Inc. All Rights Reserved.
+
+using UnrealBuildTool;
+
+public class DesktopTargetPlatform : ModuleRules
+{
+	public DesktopTargetPlatform(TargetInfo Target)
+	{
+		PrivateDependencyModuleNames.AddRange(
+			new string[]
+			{
+				"Core",
+				"CoreUObject",
+				"TargetPlatform",
+				"DesktopPlatform",
+				"LaunchDaemonMessages",
+				"Projects"
+			}
+		);
+
+		PrivateIncludePathModuleNames.AddRange(
+			new string[] {
+				"Messaging",
+				"TargetDeviceServices",
+			}
+		);
+
+		if (UEBuildConfiguration.bCompileAgainstEngine)
+		{
+			PrivateDependencyModuleNames.Add("Engine");
+		}
+	}
+}
diff --git a/Engine/Source/Developer/DesktopTargetPlatform/Private/DesktopTargetPlatform.cpp b/Engine/Source/Developer/DesktopTargetPlatform/Private/DesktopTargetPlatform.cpp
new file mode 100644
--- /dev/null
+++ b/Engine/Source/Developer/DesktopTargetPlatform/Private/DesktopTargetPlatform.cpp
@@ -0,0 +1,85 @@
+// Copyright 1998-2015 Epic Games, Inc. All Rights Reserved.
+
+/*=============================================================================
+	DesktopTargetPlatform.cpp: Implements the FDesktopTargetPlatform class.
+=============================================================================*/
+
+#include "DesktopTargetPlatformPrivatePCH.h"
+#include "IProjectManager.h"
+
+/* FDesktopTargetPlatform structors
+ *****************************************************************************/
+
+FDesktopTargetPlatform::FDesktopTargetPlatform()
+{
+#if WITH_ENGINE
+	// use non-platform specific settings
+	FConfigCacheIni::LoadLocalIniFile(EngineSettings, TEXT("Engine"), true, NULL);
+	TextureLODSettings.Initialize(EngineSettings, TEXT("SystemSettings"));
+	StaticMeshLODSettings.Initialize(EngineSettings);
+#endif // #if WITH_ENGINE
+
+}
+
+
+FDesktopTargetPlatform::~FDesktopTargetPlatform()
+{
+}
+
+
+
+/* ITargetPlatform interface
+ *****************************************************************************/
+
+#if WITH_ENGINE
+
+
+void FDesktopTargetPlatform::GetAllPossibleShaderFormats( TArray<FName>& OutFormats ) const
+{
+	static FName NAME_PCD3D_SM5(TEXT("PCD3D_SM5"));
+	static FName NAME_PCD3D_SM4(TEXT("PCD3D_SM4"));
+	static FName NAME_GLSL_150(TEXT("GLSL_150"));
+	static FName NAME_GLSL_430(TEXT("GLSL_430"));
+	static FName NAME_GLSL_150_MAC(TEXT("GLSL_150_MAC"));
+	
+//	OutFormats.AddUnique(NAME_PCD3D_SM5);
+//	OutFormats.AddUnique(NAME_PCD3D_SM4);
+	OutFormats.AddUnique(NAME_GLSL_150);
+	OutFormats.AddUnique(NAME_GLSL_430);
+	OutFormats.AddUnique(NAME_GLSL_150_MAC);
+}
+
+
+void FDesktopTargetPlatform::GetAllTargetedShaderFormats( TArray<FName>& OutFormats ) const
+{
+	GetAllPossibleShaderFormats(OutFormats);
+}
+
+void FDesktopTargetPlatform::GetTextureFormats( const UTexture* Texture, TArray<FName>& OutFormats ) const
+{
+	// just use the standard texture format name for this texture (without DX11 texture support)
+	OutFormats.Add(GetDefaultTextureFormatName(Texture, EngineSettings, false));
+}
+
+
+const FTextureLODSettings& FDesktopTargetPlatform::GetTextureLODSettings() const
+{
+	return TextureLODSettings;
+}
+
+
+FName FDesktopTargetPlatform::GetWaveFormat( const class USoundWave* Wave ) const
+{
+	static FName NAME_OGG(TEXT("OGG"));
+	static FName NAME_OPUS(TEXT("OPUS"));
+	
+	if (Wave->IsStreaming())
+	{
+		// @todo desktop platform: Does Mac support OPUS?
+		return NAME_OPUS;
+	}
+	
+	return NAME_OGG;
+}
+
+#endif // WITH_ENGINE
diff --git a/Engine/Source/Developer/DesktopTargetPlatform/Private/DesktopTargetPlatform.h b/Engine/Source/Developer/DesktopTargetPlatform/Private/DesktopTargetPlatform.h
new file mode 100644
--- /dev/null
+++ b/Engine/Source/Developer/DesktopTargetPlatform/Private/DesktopTargetPlatform.h
@@ -0,0 +1,166 @@
+// Copyright 1998-2015 Epic Games, Inc. All Rights Reserved.
+
+/*=============================================================================
+	DesktopTargetPlatform.h: Declares the FDesktopTargetPlatform class.
+=============================================================================*/
+
+#pragma once
+
+#if WITH_ENGINE
+#include "StaticMeshResources.h"
+#endif // WITH_ENGINE
+
+class FDesktopPlatformProperties : public FGenericPlatformProperties
+{
+public:
+	static FORCEINLINE const char* PlatformName()
+	{
+		return "Desktop";
+	}
+	static FORCEINLINE const char* IniPlatformName()
+	{
+		// this uses generic, non-platform-specific .ini files
+		return "";
+	}
+	static FORCEINLINE bool HasEditorOnlyData()
+	{
+		return false;
+	}
+	static FORCEINLINE bool RequiresCookedData()
+	{
+		return true;
+	}
+};
+
+/**
+ * FDesktopTargetPlatform, abstraction for cooking iOS platforms
+ */
+class FDesktopTargetPlatform
+	: public TTargetPlatformBase<FDesktopPlatformProperties>
+{
+public:
+
+	/**
+	 * Default constructor.
+	 */
+	FDesktopTargetPlatform();
+
+	/**
+	 * Destructor.
+	 */
+	~FDesktopTargetPlatform();
+
+public:
+
+	// Begin TTargetPlatformBase interface
+
+	virtual bool IsServerOnly( ) const override
+	{
+		return false;
+	}
+
+	// End TTargetPlatformBase interface
+
+public:
+
+	// Begin ITargetPlatform interface
+
+	virtual void EnableDeviceCheck(bool OnOff) override
+	{
+	}
+
+	virtual void GetAllDevices( TArray<ITargetDevicePtr>& OutDevices ) const override
+	{
+	}
+
+	virtual ECompressionFlags GetBaseCompressionMethod() const override
+	{
+		return COMPRESS_ZLIB;
+	}
+
+	virtual bool GenerateStreamingInstallManifest(const TMultiMap<FString, int32>& ChunkMap, const TSet<int32>& ChunkIDsInUse) const override
+	{
+		return true;
+	}
+
+	virtual ITargetDevicePtr GetDefaultDevice( ) const override
+	{
+		return NULL;
+	}
+
+	virtual ITargetDevicePtr GetDevice( const FTargetDeviceId& DeviceId ) override
+	{
+		return NULL;
+	}
+
+	virtual bool IsRunningPlatform( ) const override
+	{
+		return false;
+	}
+
+	virtual bool SupportsFeature( ETargetPlatformFeatures Feature ) const override
+	{
+		return TTargetPlatformBase<FDesktopPlatformProperties>::SupportsFeature(Feature);
+	}
+
+	virtual bool IsSdkInstalled(bool bProjectHasCode, FString& OutTutorialPath) const override
+	{
+		return true;
+	}
+
+	virtual int32 CheckRequirements(const FString& ProjectPath, bool bProjectHasCode, FString& OutTutorialPath) const override
+	{
+		return ETargetPlatformReadyStatus::Ready;
+	}
+
+
+#if WITH_ENGINE
+
+	virtual void GetAllPossibleShaderFormats( TArray<FName>& OutFormats ) const override;
+
+	virtual void GetAllTargetedShaderFormats( TArray<FName>& OutFormats ) const override;
+
+	virtual const class FStaticMeshLODSettings& GetStaticMeshLODSettings( ) const override
+	{
+		return StaticMeshLODSettings;
+	}
+
+	virtual void GetTextureFormats( const UTexture* Texture, TArray<FName>& OutFormats ) const override;
+
+	virtual const struct FTextureLODSettings& GetTextureLODSettings( ) const override;
+
+	virtual FName GetWaveFormat( const class USoundWave* Wave ) const override;
+#endif // WITH_ENGINE
+
+
+	DECLARE_DERIVED_EVENT(FDesktopTargetPlatform, ITargetPlatform::FOnTargetDeviceDiscovered, FOnTargetDeviceDiscovered);
+	virtual FOnTargetDeviceDiscovered& OnDeviceDiscovered( ) override
+	{
+		return DeviceDiscoveredEvent;
+	}
+
+	DECLARE_DERIVED_EVENT(FDesktopTargetPlatform, ITargetPlatform::FOnTargetDeviceLost, FOnTargetDeviceLost);
+	virtual FOnTargetDeviceLost& OnDeviceLost( ) override
+	{
+		return DeviceLostEvent;
+	}
+
+private:
+
+#if WITH_ENGINE
+	// Holds the Engine INI settings, for quick use.
+	FConfigFile EngineSettings;
+
+	// Holds the cache of the target LOD settings.
+	FTextureLODSettings TextureLODSettings;
+
+	// Holds the static mesh LOD settings.
+	FStaticMeshLODSettings StaticMeshLODSettings;
+#endif // WITH_ENGINE
+
+	// Holds an event delegate that is executed when a new target device has been discovered.
+	FOnTargetDeviceDiscovered DeviceDiscoveredEvent;
+
+	// Holds an event delegate that is executed when a target device has been lost, i.e. disconnected or timed out.
+	FOnTargetDeviceLost DeviceLostEvent;
+};
diff --git a/Engine/Source/Developer/DesktopTargetPlatform/Private/DesktopTargetPlatformModule.cpp b/Engine/Source/Developer/DesktopTargetPlatform/Private/DesktopTargetPlatformModule.cpp
new file mode 100644
--- /dev/null
+++ b/Engine/Source/Developer/DesktopTargetPlatform/Private/DesktopTargetPlatformModule.cpp
@@ -0,0 +1,57 @@
+// Copyright 1998-2015 Epic Games, Inc. All Rights Reserved.
+
+#include "DesktopTargetPlatformPrivatePCH.h"
+
+/**
+ * Holds the target platform singleton.
+ */
+static ITargetPlatform* Singleton = NULL;
+
+
+/**
+ * Module for a generic target platform for desktop platforms
+ */
+class FDesktopTargetPlatformModule : public ITargetPlatformModule
+{
+public:
+
+	/**
+	 * Destructor.
+	 */
+	~FDesktopTargetPlatformModule()
+	{
+		Singleton = NULL;
+	}
+
+
+public:
+
+	// Begin ITargetPlatformModule interface
+
+	virtual ITargetPlatform* GetTargetPlatform()
+	{
+		if (Singleton == NULL)
+		{
+			Singleton = new FDesktopTargetPlatform();
+		}
+
+		return Singleton;
+	}
+
+	// End ITargetPlatformModule interface
+
+
+public:
+
+	// Begin IModuleInterface interface
+	virtual void StartupModule() override
+	{
+	}
+
+	virtual void ShutdownModule() override
+	{
+	}
+	// End IModuleInterface interface
+};
+
+IMPLEMENT_MODULE( FDesktopTargetPlatformModule, DesktopTargetPlatform);
diff --git a/Engine/Source/Developer/DesktopTargetPlatform/Private/DesktopTargetPlatformPrivatePCH.h b/Engine/Source/Developer/DesktopTargetPlatform/Private/DesktopTargetPlatformPrivatePCH.h
new file mode 100644
--- /dev/null
+++ b/Engine/Source/Developer/DesktopTargetPlatform/Private/DesktopTargetPlatformPrivatePCH.h
@@ -0,0 +1,33 @@
+// Copyright 1998-2015 Epic Games, Inc. All Rights Reserved.
+
+/*=============================================================================
+	IOSTargetPlatformPrivatePCH.h: Pre-compiled header file for the IOSTargetPlatform module.
+=============================================================================*/
+
+#pragma once
+
+
+/* Dependencies
+ *****************************************************************************/
+
+#include "Core.h"
+#include "ModuleInterface.h"
+#include "ModuleManager.h"
+#include "LaunchDaemonMessages.h"
+#include "Messaging.h"
+#include "TargetDeviceServices.h"
+#include "GenericPlatformProperties.h"
+
+#if WITH_ENGINE
+	#include "Engine.h"
+#endif
+
+#include "TargetPlatform.h"
+#include "TargetPlatformBase.h"
+
+
+/* Private includes
+ *****************************************************************************/
+
+#include "DesktopTargetPlatform.h"
+ 
\ No newline at end of file
diff --git a/Engine/Source/Developer/Linux/LinuxTargetPlatform/Private/LinuxTargetPlatform.h b/Engine/Source/Developer/Linux/LinuxTargetPlatform/Private/LinuxTargetPlatform.h
--- a/Engine/Source/Developer/Linux/LinuxTargetPlatform/Private/LinuxTargetPlatform.h
+++ b/Engine/Source/Developer/Linux/LinuxTargetPlatform/Private/LinuxTargetPlatform.h
@@ -161,22 +161,10 @@ public:
 
 	virtual void GetTextureFormats( const UTexture* InTexture, TArray<FName>& OutFormats ) const override
 	{
-		static FName NameBC6H(TEXT("BC6H"));
-		static FName NameBC7(TEXT("BC7"));
-		static FName NameRGBA16F(TEXT("RGBA16F"));
-		static FName NameAutoDXT(TEXT("AutoDXT"));
 		if (!IS_DEDICATED_SERVER)
 		{
 			// just use the standard texture format name for this texture
-			FName TextureFormatName = this->GetDefaultTextureFormatName(InTexture, EngineSettings);
-			if (TextureFormatName == NameBC6H)
-			{
-				TextureFormatName = NameRGBA16F;
-			}
-			else if (TextureFormatName == NameBC7)
-			{
-				TextureFormatName = NameAutoDXT;
-			}
+			FName TextureFormatName = this->GetDefaultTextureFormatName(InTexture, EngineSettings, false);
 			OutFormats.Add(TextureFormatName);
 		}
 	}
diff --git a/Engine/Source/Developer/Mac/MacTargetPlatform/Private/GenericMacTargetPlatform.h b/Engine/Source/Developer/Mac/MacTargetPlatform/Private/GenericMacTargetPlatform.h
--- a/Engine/Source/Developer/Mac/MacTargetPlatform/Private/GenericMacTargetPlatform.h
+++ b/Engine/Source/Developer/Mac/MacTargetPlatform/Private/GenericMacTargetPlatform.h
@@ -114,30 +114,12 @@ return TSuper::SupportsFeature(Feature);
 		return StaticMeshLODSettings;
 	}
 
-	virtual void GetTextureFormats( const UTexture* InTexture, TArray<FName>& OutFormats ) const override
+	virtual void GetTextureFormats( const UTexture* Texture, TArray<FName>& OutFormats ) const override
 	{
-		static FName NameBC6H(TEXT("BC6H"));
-		static FName NameBC7(TEXT("BC7"));
-		static FName NameRGBA16F(TEXT("RGBA16F"));
-		static FName NameAutoDXT(TEXT("AutoDXT"));
-
-		if (IS_DEDICATED_SERVER)
-		{
-			// no textures needed for dedicated server target
-			OutFormats.Add(NAME_None);
-		}
-		else
+		if (!IS_DEDICATED_SERVER)
 		{
-			// just use the standard texture format name for this texture
-			FName TextureFormatName = this->GetDefaultTextureFormatName(InTexture, EngineSettings);
-			if (TextureFormatName == NameBC6H)
-			{
-				TextureFormatName = NameRGBA16F;
-			}
-			else if (TextureFormatName == NameBC7)
-			{
-				TextureFormatName = NameAutoDXT;
-			}
+			// just use the standard texture format name for this texture (with no DX11 support)
+			FName TextureFormatName = this->GetDefaultTextureFormatName(Texture, EngineSettings, false);
 			OutFormats.Add(TextureFormatName);
 		}
 	}
diff --git a/Engine/Source/Developer/TargetPlatform/Public/Common/TargetPlatformBase.h b/Engine/Source/Developer/TargetPlatform/Public/Common/TargetPlatformBase.h
--- a/Engine/Source/Developer/TargetPlatform/Public/Common/TargetPlatformBase.h
+++ b/Engine/Source/Developer/TargetPlatform/Public/Common/TargetPlatformBase.h
@@ -37,7 +37,7 @@ public:
 	}
 
 #ifdef TEXTURERESOURCE_H_INCLUDED // defined in TextureResource.h, this way we know if UTexture is available, needed for Clang
-	FName GetDefaultTextureFormatName( const UTexture* Texture, const FConfigFile& EngineSettings ) const
+	FName GetDefaultTextureFormatName( const UTexture* Texture, const FConfigFile& EngineSettings, bool bSupportDX11TextureFormats ) const
 	{
 		FName TextureFormatName = NAME_None;
 
@@ -155,6 +155,20 @@ public:
 		{
 				TextureFormatName = NameBGRA8;
 		}
+
+		// fallback to non-DX11 formats if one was chosen, but we can't use it
+		if (!bSupportDX11TextureFormats)
+		{
+			if (TextureFormatName == NameBC6H)
+			{
+				TextureFormatName = NameRGBA16F;
+			}
+			else if (TextureFormatName == NameBC7)
+			{
+				TextureFormatName = NameAutoDXT;
+			}
+		}
+
 #endif //WITH_EDITOR
 
 		return TextureFormatName;
diff --git a/Engine/Source/Developer/TargetPlatform/TargetPlatform.Build.cs b/Engine/Source/Developer/TargetPlatform/TargetPlatform.Build.cs
--- a/Engine/Source/Developer/TargetPlatform/TargetPlatform.Build.cs
+++ b/Engine/Source/Developer/TargetPlatform/TargetPlatform.Build.cs
@@ -78,6 +78,7 @@ public class TargetPlatform : ModuleRules
 				if (UEBuildConfiguration.bCompileAgainstEngine)
 				{
 					DynamicallyLoadedModuleNames.Add("AudioFormatOgg");
+					DynamicallyLoadedModuleNames.Add("AudioFormatOpus");
 				}
 
 				if (Target.Type == TargetRules.TargetType.Editor || Target.Type == TargetRules.TargetType.Program)
diff --git a/Engine/Source/Developer/Windows/WindowsTargetPlatform/Private/GenericWindowsTargetPlatform.h b/Engine/Source/Developer/Windows/WindowsTargetPlatform/Private/GenericWindowsTargetPlatform.h
--- a/Engine/Source/Developer/Windows/WindowsTargetPlatform/Private/GenericWindowsTargetPlatform.h
+++ b/Engine/Source/Developer/Windows/WindowsTargetPlatform/Private/GenericWindowsTargetPlatform.h
@@ -142,25 +142,9 @@ public:
 
 	virtual void GetTextureFormats( const UTexture* InTexture, TArray<FName>& OutFormats ) const override
 	{
-		static FName NameBC6H(TEXT("BC6H"));
-		static FName NameBC7(TEXT("BC7"));
-		static FName NameRGBA16F(TEXT("RGBA16F"));
-		static FName NameAutoDXT(TEXT("AutoDXT"));
-
 		if (!IS_DEDICATED_SERVER)
 		{
-			FName TextureFormatName = GetDefaultTextureFormatName(InTexture, EngineSettings);
-			if (!bSupportDX11TextureFormats)
-			{
-				if (TextureFormatName == NameBC6H)
-				{
-					TextureFormatName = NameRGBA16F;
-				}
-				else if (TextureFormatName == NameBC7)
-				{
-					TextureFormatName = NameAutoDXT;
-				}
-			}
+			FName TextureFormatName = GetDefaultTextureFormatName(InTexture, EngineSettings, bSupportDX11TextureFormats);
 			OutFormats.Add(TextureFormatName);
 		}
 	}
diff --git a/Engine/Source/Developer/iOS/IOSTargetPlatform/Private/IOSTargetPlatform.cpp b/Engine/Source/Developer/iOS/IOSTargetPlatform/Private/IOSTargetPlatform.cpp
--- a/Engine/Source/Developer/iOS/IOSTargetPlatform/Private/IOSTargetPlatform.cpp
+++ b/Engine/Source/Developer/iOS/IOSTargetPlatform/Private/IOSTargetPlatform.cpp
@@ -375,8 +375,6 @@ void FIOSTargetPlatform::GetTextureFormats( const UTexture* Texture, TArray<FNam
 		FName(TEXT("DXT5n")),	FName(TEXT("PVRTCN")),		FName(TEXT("ASTC_NormalAG")),
 		FName(TEXT("BC5")),		FName(TEXT("PVRTCN")),		FName(TEXT("ASTC_NormalRG")),
 		FName(TEXT("AutoDXT")),	FName(TEXT("AutoPVRTC")),	FName(TEXT("ASTC_RGBAuto")),
-		FName(TEXT("BC6H")),	FName(TEXT("RGBA16F")),		FName(TEXT("RGBA16F")),
-		FName(TEXT("BC7")),		FName(TEXT("AutoPVRTC")),	FName(TEXT("ASTC_RGBAuto")),
 	};
 
 	FName TextureFormatName = NAME_None;
@@ -390,7 +388,7 @@ void FIOSTargetPlatform::GetTextureFormats( const UTexture* Texture, TArray<FNam
 	// if we didn't assign anything specially, then use the defaults
 	if (TextureFormatName == NAME_None)
 	{
-		TextureFormatName = GetDefaultTextureFormatName(Texture, EngineSettings);
+		TextureFormatName = GetDefaultTextureFormatName(Texture, EngineSettings, false);
 	}
 
 	// perform any remapping away from defaults
diff --git a/Engine/Source/Programs/AutomationTool/Desktop/Desktop.Automation.csproj b/Engine/Source/Programs/AutomationTool/Desktop/Desktop.Automation.csproj
new file mode 100644
--- /dev/null
+++ b/Engine/Source/Programs/AutomationTool/Desktop/Desktop.Automation.csproj
@@ -0,0 +1,75 @@
+<?xml version="1.0" encoding="utf-8"?>
+<Project ToolsVersion="4.0" DefaultTargets="Build" xmlns="http://schemas.microsoft.com/developer/msbuild/2003">
+  <Import Project="$(MSBuildExtensionsPath)\$(MSBuildToolsVersion)\Microsoft.Common.props" Condition="Exists('$(MSBuildExtensionsPath)\$(MSBuildToolsVersion)\Microsoft.Common.props')" />
+  <PropertyGroup>
+    <Configuration Condition=" '$(Configuration)' == '' ">Debug</Configuration>
+    <Platform Condition=" '$(Platform)' == '' ">AnyCPU</Platform>
+    <ProjectGuid>{32E3A47E-359A-4F78-AE6C-703B87DE0B97}</ProjectGuid>
+    <OutputType>Library</OutputType>
+    <AppDesignerFolder>Properties</AppDesignerFolder>
+    <RootNamespace>Desktop.Automation</RootNamespace>
+    <AssemblyName>Desktop.Automation</AssemblyName>
+    <TargetFrameworkVersion>v4.0</TargetFrameworkVersion>
+    <FileAlignment>512</FileAlignment>
+  </PropertyGroup>
+  <PropertyGroup Condition=" '$(Configuration)|$(Platform)' == 'Debug|AnyCPU' ">
+    <DebugSymbols>true</DebugSymbols>
+    <DebugType>full</DebugType>
+    <Optimize>false</Optimize>
+    <OutputPath>..\..\..\..\Binaries\DotNET\AutomationScripts\Desktop\</OutputPath>
+    <DefineConstants>DEBUG;TRACE</DefineConstants>
+    <ErrorReport>prompt</ErrorReport>
+    <WarningLevel>4</WarningLevel>
+    <PlatformTarget>AnyCPU</PlatformTarget>
+    <TreatWarningsAsErrors>true</TreatWarningsAsErrors>
+  </PropertyGroup>
+  <PropertyGroup Condition=" '$(Configuration)|$(Platform)' == 'Development|AnyCPU' ">
+    <DebugType>pdbonly</DebugType>
+    <Optimize>true</Optimize>
+    <OutputPath>..\..\..\..\Binaries\DotNET\AutomationScripts\Desktop\</OutputPath>
+    <DefineConstants>TRACE</DefineConstants>
+    <ErrorReport>prompt</ErrorReport>
+    <WarningLevel>4</WarningLevel>
+    <PlatformTarget>AnyCPU</PlatformTarget>
+    <TreatWarningsAsErrors>true</TreatWarningsAsErrors>
+  </PropertyGroup>
+  <ItemGroup>
+    <Reference Include="Ionic.Zip.Reduced">
+      <HintPath>..\..\..\..\Binaries\DotNET\Ionic.Zip.Reduced.dll</HintPath>
+    </Reference>
+    <Reference Include="System" />
+    <Reference Include="System.Core" />
+    <Reference Include="System.Runtime.Serialization" />
+    <Reference Include="System.Web" />
+    <Reference Include="System.Web.Extensions" />
+    <Reference Include="System.Xml.Linq" />
+    <Reference Include="System.Data.DataSetExtensions" />
+    <Reference Include="Microsoft.CSharp" />
+    <Reference Include="System.Data" />
+    <Reference Include="System.Xml" />
+    <Reference Include="OneSky">
+      <HintPath>..\..\..\..\Binaries\DotNET\OneSky.dll</HintPath>
+      <Private>False</Private>
+    </Reference>
+  </ItemGroup>
+  <ItemGroup>
+    <Compile Include="..\..\DotNETCommon\MetaData.cs">
+      <Link>Properties\MetaData.cs</Link>
+    </Compile>
+    <Compile Include="DesktopPlatform.Automation.cs" />
+    <Compile Include="Properties\AssemblyInfo.cs" />
+  </ItemGroup>
+  <ItemGroup>
+    <ProjectReference Include="..\AutomationTool.csproj">
+      <Project>{A35B0261-D868-475E-A098-EB14CE6DC89F}</Project>
+      <Name>AutomationTool</Name>
+      <Private>False</Private>
+    </ProjectReference>
+    <ProjectReference Include="..\..\UnrealBuildTool\UnrealBuildTool.csproj">
+      <Project>{FD7C5E1A-CFE4-4FD5-A525-1EB1599A39AC}</Project>
+      <Name>UnrealBuildTool</Name>
+      <Private>False</Private>
+    </ProjectReference>
+  </ItemGroup>
+  <Import Project="$(MSBuildToolsPath)\Microsoft.CSharp.targets" />
+</Project>
diff --git a/Engine/Source/Programs/AutomationTool/Desktop/DesktopPlatform.Automation.cs b/Engine/Source/Programs/AutomationTool/Desktop/DesktopPlatform.Automation.cs
new file mode 100644
--- /dev/null
+++ b/Engine/Source/Programs/AutomationTool/Desktop/DesktopPlatform.Automation.cs
@@ -0,0 +1,85 @@
+﻿// Copyright 1998-2015 Epic Games, Inc. All Rights Reserved.
+using System;
+using System.Collections.Generic;
+using System.Linq;
+using System.Text;
+using System.IO;
+using System.Diagnostics;
+using System.Net.NetworkInformation;
+using System.Threading;
+using AutomationTool;
+using UnrealBuildTool;
+using Ionic.Zip;
+
+public class DesktopPlatform : Platform
+{
+	public DesktopPlatform()
+		: base(UnrealTargetPlatform.Desktop)
+	{
+	}
+
+	public override void Package(ProjectParams Params, DeploymentContext SC, int WorkingCL)
+	{
+
+		PrintRunTime();
+	}
+
+	public override void GetFilesToArchive(ProjectParams Params, DeploymentContext SC)
+	{
+	}
+
+	public override void GetConnectedDevices(ProjectParams Params, out List<string> Devices)
+	{
+		Devices = new List<string>();
+	}
+
+	public override void Deploy(ProjectParams Params, DeploymentContext SC)
+	{
+	}
+
+	public override ProcessResult RunClient(ERunOptions ClientRunFlags, string ClientApp, string ClientCmdLine, ProjectParams Params)
+	{
+		return null;
+	}
+
+	public override void GetFilesToDeployOrStage(ProjectParams Params, DeploymentContext SC)
+	{
+	}
+
+	/// <summary>
+	/// Gets cook platform name for this platform.
+	/// </summary>
+	/// <param name="CookFlavor">Additional parameter used to indicate special sub-target platform.</param>
+	/// <returns>Cook platform string.</returns>
+	public override string GetCookPlatform(bool bDedicatedServer, bool bIsClientOnly, string CookFlavor)
+	{
+		return "Desktop";
+	}
+
+	public override bool DeployPakInternalLowerCaseFilenames()
+	{
+		return false;
+	}
+
+	public override bool DeployLowerCaseFilenames(bool bUFSFile)
+	{
+		return false;
+	}
+
+	public override bool IsSupported { get { return true; } }
+
+	public override string Remap(string Dest)
+	{
+		return Dest;
+	}
+
+	public override PakType RequiresPak(ProjectParams Params)
+	{
+		return PakType.DontCare;
+	}
+    
+	public override List<string> GetDebugFileExtentions()
+	{
+		return new List<string> { };
+	}
+}
diff --git a/Engine/Source/Programs/AutomationTool/Desktop/Properties/AssemblyInfo.cs b/Engine/Source/Programs/AutomationTool/Desktop/Properties/AssemblyInfo.cs
new file mode 100644
--- /dev/null
+++ b/Engine/Source/Programs/AutomationTool/Desktop/Properties/AssemblyInfo.cs
@@ -0,0 +1,12 @@
+// Copyright 1998-2015 Epic Games, Inc. All Rights Reserved.
+
+using System.Reflection;
+using System.Runtime.CompilerServices;
+using System.Runtime.InteropServices;
+
+// General Information about an assembly is controlled through the following 
+// set of attributes. Change these attribute values to modify the information
+// associated with an assembly.
+[assembly: AssemblyTitle("Desktop.Automation")]
+[assembly: AssemblyDescription("")]
+[assembly: AssemblyConfiguration("")]
diff --git a/Engine/Source/Programs/UnrealBuildTool/Configuration/UEBuildTarget.cs b/Engine/Source/Programs/UnrealBuildTool/Configuration/UEBuildTarget.cs
--- a/Engine/Source/Programs/UnrealBuildTool/Configuration/UEBuildTarget.cs
+++ b/Engine/Source/Programs/UnrealBuildTool/Configuration/UEBuildTarget.cs
@@ -24,6 +24,7 @@ namespace UnrealBuildTool
 		WinRT_ARM,
 		HTML5,
         Linux,
+		Desktop,
 	}
 
 	public enum UnrealPlatformGroup
@@ -41,7 +42,8 @@ namespace UnrealBuildTool
           *  Note: There's no default group - if the platform is not registered as device or simulator - both are rejected. 
           */
         Device, 
-        Simulator, 
+        Simulator,
+		Desktop,
 	}
 
 	public enum UnrealTargetConfiguration
@@ -787,7 +789,8 @@ namespace UnrealBuildTool
 			ProjectIntermediateDirectory = Path.GetFullPath(Path.Combine(ProjectDirectory, BuildConfiguration.PlatformIntermediateFolder, GetTargetName(), Configuration.ToString()));
 
 			// Build the engine intermediate directory. If we're building agnostic engine binaries, we can use the engine intermediates folder. Otherwise we need to use the project intermediates directory.
-			if (ShouldCompileMonolithic())
+			// If we are building a content project with UE4Game as the base executable, then we don't want to use the Project as the intermediate, it should use UE4's intermedate
+			if (ShouldCompileMonolithic() && AppName != "UE4Game")
 			{
 				EngineIntermediateDirectory = ProjectIntermediateDirectory;
 			}
diff --git a/Engine/Source/Programs/UnrealBuildTool/Linux/UEBuildLinux.cs b/Engine/Source/Programs/UnrealBuildTool/Linux/UEBuildLinux.cs
--- a/Engine/Source/Programs/UnrealBuildTool/Linux/UEBuildLinux.cs
+++ b/Engine/Source/Programs/UnrealBuildTool/Linux/UEBuildLinux.cs
@@ -331,6 +331,7 @@ namespace UnrealBuildTool
 
                 if (bBuildShaderFormats)
                 {
+					InModule.AddDynamicallyLoadedModule("ShaderFormatD3D");
                     InModule.AddDynamicallyLoadedModule("ShaderFormatOpenGL");
                 }
             }
diff --git a/Engine/Source/Programs/UnrealBuildTool/Mac/UEBuildMac.cs b/Engine/Source/Programs/UnrealBuildTool/Mac/UEBuildMac.cs
--- a/Engine/Source/Programs/UnrealBuildTool/Mac/UEBuildMac.cs
+++ b/Engine/Source/Programs/UnrealBuildTool/Mac/UEBuildMac.cs
@@ -117,6 +117,7 @@ namespace UnrealBuildTool
 
                 if (bBuildShaderFormats)
                 {
+					InModule.AddDynamicallyLoadedModule("ShaderFormatD3D");
                     InModule.AddDynamicallyLoadedModule("ShaderFormatOpenGL");
                 }
 			}
diff --git a/Engine/Source/Programs/UnrealBuildTool/Windows/UEBuildWindows.cs b/Engine/Source/Programs/UnrealBuildTool/Windows/UEBuildWindows.cs
--- a/Engine/Source/Programs/UnrealBuildTool/Windows/UEBuildWindows.cs
+++ b/Engine/Source/Programs/UnrealBuildTool/Windows/UEBuildWindows.cs
@@ -415,6 +415,7 @@ namespace UnrealBuildTool
                     InModule.AddDynamicallyLoadedModule("WindowsNoEditorTargetPlatform");
                     InModule.AddDynamicallyLoadedModule("WindowsServerTargetPlatform");
                     InModule.AddDynamicallyLoadedModule("WindowsClientTargetPlatform");
+					InModule.AddDynamicallyLoadedModule("DesktopTargetPlatform");
                 }
 
                 if (bBuildShaderFormats)
diff --git a/Engine/Source/Runtime/Engine/Engine.Build.cs b/Engine/Source/Runtime/Engine/Engine.Build.cs
--- a/Engine/Source/Runtime/Engine/Engine.Build.cs
+++ b/Engine/Source/Runtime/Engine/Engine.Build.cs
@@ -142,6 +142,7 @@ public class Engine : ModuleRules
 						"WindowsNoEditorTargetPlatform",
 						"WindowsServerTargetPlatform",
 						"WindowsClientTargetPlatform",
+						"DesktopTargetPlatform",
 					}
 				);
 			}
@@ -152,7 +153,8 @@ public class Engine : ModuleRules
 					    "MacTargetPlatform",
 					    "MacNoEditorTargetPlatform",
 						"MacServerTargetPlatform",
-						"MacClientTargetPlatform"
+						"MacClientTargetPlatform",
+						"DesktopTargetPlatform",
 					}
 				);
 			}
@@ -163,6 +165,7 @@ public class Engine : ModuleRules
 						"LinuxTargetPlatform",
 						"LinuxNoEditorTargetPlatform",
 						"LinuxServerTargetPlatform",
+						"DesktopTargetPlatform",
 					}
 				);
 			}
